Fix compile error in an ICU test related to FDefaultRichTextMarkupParser::Create being replaced with FDefaultRichTextMarkupParser::Get

diff --git a/Engine/Plugins/Tests/RuntimeTests/Source/RuntimeTests/Private/Slate/RichTextMarkupProcessingTest.cpp b/Engine/Plugins/Tests/RuntimeTests/Source/RuntimeTests/Private/Slate/RichTextMarkupProcessingTest.cpp
--- a/Engine/Plugins/Tests/RuntimeTests/Source/RuntimeTests/Private/Slate/RichTextMarkupProcessingTest.cpp
+++ b/Engine/Plugins/Tests/RuntimeTests/Source/RuntimeTests/Private/Slate/RichTextMarkupProcessingTest.cpp
@@ -100,7 +100,7 @@ namespace
 bool FRTFProcessingTest::RunTest (const FString& Parameters)
 {
 #if UE_ENABLE_ICU
-	TSharedRef< FDefaultRichTextMarkupParser > RichTextMarkupProcessing = FDefaultRichTextMarkupParser::Create();
+	TSharedRef< FDefaultRichTextMarkupParser > RichTextMarkupProcessing = FDefaultRichTextMarkupParser::Get();
 
 	// Text
 	{
