Change the RuntimeTests plugin to also work in cooked builds, by conditionally compiling the image comparison tools used by the disabled shader comparison test
This plugin will be used for new engine runtime tests

diff --git a/Engine/Plugins/Tests/RuntimeTests/RuntimeTests.uplugin b/Engine/Plugins/Tests/RuntimeTests/RuntimeTests.uplugin
--- a/Engine/Plugins/Tests/RuntimeTests/RuntimeTests.uplugin
+++ b/Engine/Plugins/Tests/RuntimeTests/RuntimeTests.uplugin
@@ -3,7 +3,7 @@
 	"Version": 1,
 	"VersionName": "1.0",
 	"FriendlyName": "Runtime Tests",
-	"Description": "",
+	"Description": "Automated tests for runtime code used in both editor and cooked builds",
 	"Category": "Testing",
 	"CreatedBy": "Epic Games, Inc.",
 	"CreatedByURL": "https://epicgames.com",
@@ -17,7 +17,7 @@
 	"Modules": [
 		{
 			"Name": "RuntimeTests",
-			"Type": "UncookedOnly",
+			"Type": "Runtime",
 			"LoadingPhase": "PreDefault"
 		}
 	]
diff --git a/Engine/Plugins/Tests/RuntimeTests/Source/RuntimeTests/Private/RuntimeTests.cpp b/Engine/Plugins/Tests/RuntimeTests/Source/RuntimeTests/Private/RuntimeTests.cpp
--- a/Engine/Plugins/Tests/RuntimeTests/Source/RuntimeTests/Private/RuntimeTests.cpp
+++ b/Engine/Plugins/Tests/RuntimeTests/Source/RuntimeTests/Private/RuntimeTests.cpp
@@ -4,7 +4,7 @@
 #include "Modules/ModuleManager.h"
 #include "ShaderComparisonTests.h"
 
-#define LOCTEXT_NAMESPACE "FRuntimeTestsModule"
+#if WITH_AUTOMATION_TESTS
 
 void FRuntimeTestsModule::StartupModule()
 {
@@ -20,6 +20,6 @@ void FRuntimeTestsModule::ShutdownModule()
 	delete FCompareBasepassShadersAutomationTestInstance;
 }
 
-#undef LOCTEXT_NAMESPACE
+#endif // WITH_AUTIOMATION_TESTS
 	
 IMPLEMENT_MODULE(FRuntimeTestsModule, RuntimeTests)
diff --git a/Engine/Plugins/Tests/RuntimeTests/Source/RuntimeTests/Private/Shaders/ShaderComparisonTests.cpp b/Engine/Plugins/Tests/RuntimeTests/Source/RuntimeTests/Private/Shaders/ShaderComparisonTests.cpp
--- a/Engine/Plugins/Tests/RuntimeTests/Source/RuntimeTests/Private/Shaders/ShaderComparisonTests.cpp
+++ b/Engine/Plugins/Tests/RuntimeTests/Source/RuntimeTests/Private/Shaders/ShaderComparisonTests.cpp
@@ -2,16 +2,21 @@
 
 #include "ShaderComparisonTests.h"
 
+#if WITH_AUTOMATION_TESTS
+
 #include "Camera/CameraActor.h"
 #include "Camera/CameraComponent.h"
 #include "GameFramework/PlayerController.h"
-#include "ImageComparer.h"
 #include "Engine/Engine.h"
 #include "Kismet/GameplayStatics.h"
 #include "Misc/FileHelper.h"
 #include "Tests/AutomationCommon.h"
 #include "UnrealClient.h"
 
+#if WITH_EDITOR
+#include "ImageComparer.h"
+#endif
+
 #include <cfloat>
 #include <memory>
 #include <exception>
@@ -228,7 +233,7 @@ bool FCompareBasepassShaders::RunTest(const FString& Parameters)
 
     // TODO(mlentine): Find appropriate normalized scores
     AddCommand(new FDelayedFunctionLatentCommand([=] {
-        FImageTolerance tolerance;
+
         float similarity = FLT_MAX;
         for (int i = 0; i < NUM_CAMERAS; ++i)
         {
@@ -238,9 +243,11 @@ bool FCompareBasepassShaders::RunTest(const FString& Parameters)
             PathName = "../../../" + PathName;
             UE_LOG(ShaderComparisonTests, Log, TEXT("Screenshots are at %s and %s."), *RealScreenshotFileName[i], *RealFP16ScreenshotFileName[i]);
             UE_LOG(ShaderComparisonTests, Log, TEXT("Difference is stored in %s"), *PathName);
+#if WITH_EDITOR
             auto single_similarity = FImageComparer().CompareStructuralSimilarity(RealScreenshotFileName[i], RealFP16ScreenshotFileName[i], FImageComparer::EStructuralSimilarityComponent::Luminance, "");
             similarity = single_similarity < similarity ? single_similarity : similarity;
             UE_LOG(ShaderComparisonTests, Log, TEXT("Similarity is %f after %d."), similarity, i);
+#endif 
         }
     
         UE_LOG(ShaderComparisonTests, Log, TEXT("FP16 ran in %f, with FP32 in %f."), (*FP16AverageTimes)[2], (*AverageTimes)[2]);
@@ -260,3 +267,5 @@ bool FCompareBasepassShaders::RunTest(const FString& Parameters)
     // This is a bit pointless right now
     return true;
 }
+
+#endif // WITH_AUTOMATION_TESTS
\ No newline at end of file
diff --git a/Engine/Plugins/Tests/RuntimeTests/Source/RuntimeTests/Public/RuntimeTests.h b/Engine/Plugins/Tests/RuntimeTests/Source/RuntimeTests/Public/RuntimeTests.h
--- a/Engine/Plugins/Tests/RuntimeTests/Source/RuntimeTests/Public/RuntimeTests.h
+++ b/Engine/Plugins/Tests/RuntimeTests/Source/RuntimeTests/Public/RuntimeTests.h
@@ -3,17 +3,21 @@
 #pragma once
 
 #include "Modules/ModuleInterface.h"
+#include "Misc/AutomationTest.h"
 
 class FCompareBasepassShaders;
 
 class FRuntimeTestsModule : public IModuleInterface
 {
-public:
+#if WITH_AUTOMATION_TESTS
 
+public:
 	/** IModuleInterface implementation */
 	virtual void StartupModule() override;
 	virtual void ShutdownModule() override;
 
 private:
 	FCompareBasepassShaders* FCompareBasepassShadersAutomationTestInstance;
+
+#endif // WITH_AUTOMATION_TESTS
 };
diff --git a/Engine/Plugins/Tests/RuntimeTests/Source/RuntimeTests/Public/ShaderComparisonTests.h b/Engine/Plugins/Tests/RuntimeTests/Source/RuntimeTests/Public/ShaderComparisonTests.h
--- a/Engine/Plugins/Tests/RuntimeTests/Source/RuntimeTests/Public/ShaderComparisonTests.h
+++ b/Engine/Plugins/Tests/RuntimeTests/Source/RuntimeTests/Public/ShaderComparisonTests.h
@@ -8,6 +8,6 @@
  * CompareAllShaderVars
  * Comparison automation test the determines which shader variables need extra precision
  */
-#if WITH_AUTOMATION_WORKER
+#if WITH_AUTOMATION_TESTS
 IMPLEMENT_COMPLEX_AUTOMATION_CLASS(FCompareBasepassShaders, "System.Engine.CompareShaderPrecision", EAutomationTestFlags::ClientContext | EAutomationTestFlags::EngineFilter | EAutomationTestFlags::Disabled)
 #endif
diff --git a/Engine/Plugins/Tests/RuntimeTests/Source/RuntimeTests/RuntimeTests.Build.cs b/Engine/Plugins/Tests/RuntimeTests/Source/RuntimeTests/RuntimeTests.Build.cs
--- a/Engine/Plugins/Tests/RuntimeTests/Source/RuntimeTests/RuntimeTests.Build.cs
+++ b/Engine/Plugins/Tests/RuntimeTests/Source/RuntimeTests/RuntimeTests.Build.cs
@@ -10,10 +10,8 @@ public class RuntimeTests : ModuleRules
 			new string[]
 			{
 				"Core",
-				
-				// ... add other public dependencies that you statically link with here ...
 			}
-			);
+		);
 			
 		
 		PrivateDependencyModuleNames.AddRange(
@@ -23,10 +21,12 @@ public class RuntimeTests : ModuleRules
 				"Engine",
 				"Slate",
 				"SlateCore",
-                "ScreenShotComparisonTools",
-
-				// ... add private dependencies that you statically link with here ...	
 			}
-			);
+		);
+
+		if (Target.bBuildEditor == true)
+		{
+			PrivateDependencyModuleNames.Add("ScreenShotComparisonTools");
+		}
 	}
 }
