Fix CIS on cooked development, apparently Misc/Paths is only implicitly included on desktop

diff --git a/Engine/Plugins/Tests/RuntimeTests/Source/RuntimeTests/Private/Shaders/ShaderComparisonTests.cpp b/Engine/Plugins/Tests/RuntimeTests/Source/RuntimeTests/Private/Shaders/ShaderComparisonTests.cpp
--- a/Engine/Plugins/Tests/RuntimeTests/Source/RuntimeTests/Private/Shaders/ShaderComparisonTests.cpp
+++ b/Engine/Plugins/Tests/RuntimeTests/Source/RuntimeTests/Private/Shaders/ShaderComparisonTests.cpp
@@ -10,6 +10,7 @@
 #include "Engine/Engine.h"
 #include "Kismet/GameplayStatics.h"
 #include "Misc/FileHelper.h"
+#include "Misc/Paths.h"
 #include "Tests/AutomationCommon.h"
 #include "UnrealClient.h"
 
