Add System.Engine.Tick.BasicTest and PerfTest automated tests to the RuntimeTests plugin, these test different methods of ticking actors for accuracy and performance comparisons.

diff --git a/Engine/Plugins/Tests/RuntimeTests/Source/RuntimeTests/Private/EngineRuntimeTests.cpp b/Engine/Plugins/Tests/RuntimeTests/Source/RuntimeTests/Private/EngineRuntimeTests.cpp
new file mode 100644
--- /dev/null
+++ b/Engine/Plugins/Tests/RuntimeTests/Source/RuntimeTests/Private/EngineRuntimeTests.cpp
@@ -0,0 +1,566 @@
+// Copyright Epic Games, Inc. All Rights Reserved.
+
+#include "EngineRuntimeTests.h"
+#include "Components/BillboardComponent.h"
+#include "Engine/Engine.h"
+#include "Engine/GameInstance.h"
+#include "Engine/Level.h"
+#include "Engine/World.h"
+#include "EngineUtils.h"
+#include "Misc/AutomationTest.h"
+#include "Stats/StatsMisc.h"
+#include "Containers/Ticker.h"
+#include "Tests/AutomationCommon.h"
+#include "TimerManager.h"
+#include "Tickable.h"
+
+AEngineTestTickActor::AEngineTestTickActor(const FObjectInitializer& ObjectInitializer)
+	: Super(ObjectInitializer)
+{
+	SpriteComponent = CreateDefaultSubobject<UBillboardComponent>(TEXT("Sprite"));
+	if (SpriteComponent)
+	{
+		SpriteComponent->bHiddenInGame = true;
+		RootComponent = SpriteComponent;
+	}
+
+	PrimaryActorTick.bCanEverTick = true;
+
+	ResetState();
+}
+
+void AEngineTestTickActor::ResetState()
+{
+	TickCount = 0;
+	bShouldIncrementTickCount = true;
+	bShouldDoMath = true;
+	MathCounter = 0.0f;
+	MathIncrement = 0.01f;
+	MathLimit = 1.0f;
+}
+
+void AEngineTestTickActor::DoTick()
+{
+	if (bShouldIncrementTickCount)
+	{
+		TickCount++;
+	}
+
+	if (bShouldDoMath && MathIncrement > 0.0f && MathLimit > 0.0f)
+	{
+		MathCounter = 0.0f;
+		while (MathCounter < MathLimit)
+		{
+			MathCounter += MathIncrement;
+		}
+	}
+}
+
+void AEngineTestTickActor::VirtualTick()
+{
+	DoTick();
+}
+
+void AEngineTestTickActor::Tick(float DeltaSeconds)
+{
+	Super::Tick(DeltaSeconds);
+
+	DoTick();
+}
+
+#if WITH_AUTOMATION_WORKER
+
+FEngineTickTestBase::FEngineTickTestBase(const FString& InName, const bool bInComplexTask)
+	: FAutomationTestBase(InName, bInComplexTask)
+{
+}
+
+FEngineTickTestBase::~FEngineTickTestBase()
+{
+	// The unique ptr takes care of destroying the world
+}
+
+UWorld* FEngineTickTestBase::GetTestWorld() const
+{
+	if (WorldWrapper.IsValid())
+	{
+		return WorldWrapper->GetTestWorld();
+	}
+	return nullptr;
+}
+
+bool FEngineTickTestBase::CreateTestWorld()
+{
+	if (!TestTrue(TEXT("TestWorld already exists in CreateTestWorld!"), GetTestWorld() == nullptr))
+	{
+		return false;
+	}
+
+	if (!WorldWrapper.IsValid())
+	{
+		WorldWrapper = TUniquePtr<FTestWorldWrapper>(new FTestWorldWrapper());
+	}
+
+	return WorldWrapper->CreateTestWorld(EWorldType::Game);
+}
+
+bool FEngineTickTestBase::CreateTestActors(int32 ActorCount, TSubclassOf<AEngineTestTickActor> ActorClass)
+{
+	UWorld* TestWorld = GetTestWorld();
+	if (!TestNotNull(TEXT("TestWorld does not exist in CreateTestActors!"), TestWorld))
+	{
+		return false;
+	}
+
+	for (int32 i = 0; i < ActorCount; i++)
+	{
+		AEngineTestTickActor* TickActor = Cast<AEngineTestTickActor>(TestWorld->SpawnActor(ActorClass.Get()));
+		if (!TestNotNull(TEXT("CreateTestActors failed to spawn actor!"), TickActor))
+		{
+			return false;
+		}
+		TickActor->ResetState();
+		TestActors.Add(TickActor);
+	}
+
+	return true;
+}
+
+bool FEngineTickTestBase::BeginPlayInTestWorld()
+{
+	UWorld* TestWorld = GetTestWorld();
+	if (!TestNotNull(TEXT("TestWorld does not exist in BeginPlayInTestWorld!"), TestWorld))
+	{
+		return false;
+	}
+
+	return WorldWrapper->BeginPlayInTestWorld();
+}
+
+bool FEngineTickTestBase::TickTestWorld(float DeltaTime)
+{
+	UWorld* TestWorld = GetTestWorld();
+	if (!TestNotNull(TEXT("TestWorld does not exist in TickTestWorld!"), TestWorld))
+	{
+		return false;
+	}
+
+	return WorldWrapper->TickTestWorld(DeltaTime);
+}
+
+bool FEngineTickTestBase::ResetTestActors()
+{
+	for (AEngineTestTickActor* TestActor : TestActors)
+	{
+		TestActor->ResetState();
+	}
+	return true;
+}
+
+bool FEngineTickTestBase::CheckTickCount(const TCHAR* TickTestName, int32 TickCount)
+{
+	for (AEngineTestTickActor* TestActor : TestActors)
+	{
+		if (!TestEqual(TickTestName, TestActor->TickCount, TickCount))
+		{
+			return false;
+		}
+	}
+	return true;
+}
+
+bool FEngineTickTestBase::DestroyAllTestActors()
+{
+	UWorld* TestWorld = GetTestWorld();
+	if (!TestNotNull(TEXT("TestWorld does not exist in CreateTestActors!"), TestWorld))
+	{
+		return false;
+	}
+
+	for (AEngineTestTickActor* TestActor : TestActors)
+	{
+		TestActor->Destroy();
+	}
+
+	TestActors.Empty();
+		
+	return true;
+}
+
+bool FEngineTickTestBase::DestroyTestWorld()
+{
+	if (WorldWrapper.IsValid())
+	{
+		DestroyAllTestActors();
+		return WorldWrapper->DestroyTestWorld(true);
+	}
+	return false;
+}
+
+bool FEngineTickTestBase::ReportAnyErrors()
+{
+	if (WorldWrapper.IsValid())
+	{
+		WorldWrapper->ForwardErrorMessages(this);
+	}
+	return HasAnyErrors();
+}
+
+
+// Emulate an efficiently registered tick with caching
+struct FEngineTestTickActorTickableFast : FTickableGameObject
+{
+	// Not safe to use outside these tests
+	AEngineTestTickActor* TickActor = nullptr;
+	UWorld* CachedWorld = nullptr;
+
+	FEngineTestTickActorTickableFast(AEngineTestTickActor* InTickActor)
+		: TickActor(InTickActor)
+	{
+		CachedWorld = InTickActor->GetWorld();
+	}
+
+	virtual void Tick(float DeltaTime) override { TickActor->DoTick(); }
+	virtual UWorld* GetTickableGameObjectWorld() const override { return CachedWorld; }
+	virtual bool IsTickableWhenPaused() const override { return false; }
+	virtual bool IsTickableInEditor() const override { return false; }
+	virtual ETickableTickType GetTickableTickType() const override { return ETickableTickType::Always; }
+	virtual TStatId GetStatId() const override { return TStatId(); }
+};
+
+// Emulates a safer and slower setup
+struct FEngineTestTickActorTickableSlow : FTickableGameObject
+{
+	// Not safe to use outside these tests
+	AEngineTestTickActor* TickActor = nullptr;
+
+	FEngineTestTickActorTickableSlow(AEngineTestTickActor* InTickActor)
+		: TickActor(InTickActor)
+	{
+	}
+
+	virtual void Tick(float DeltaTime) override { TickActor->VirtualTick(); }
+	virtual UWorld* GetTickableGameObjectWorld() const override { return TickActor->GetWorld(); }
+	virtual bool IsTickableWhenPaused() const override { return false; }
+	virtual bool IsTickableInEditor() const override { return false; }
+	virtual bool IsAllowedToTick() const { return IsValid(TickActor) && IsValid(TickActor->GetOuter()); }
+	virtual bool IsTickable() const { return TickActor->bShouldIncrementTickCount; }
+	virtual ETickableTickType GetTickableTickType() const override { return ETickableTickType::Conditional; }
+	virtual TStatId GetStatId() const override { RETURN_QUICK_DECLARE_CYCLE_STAT(FEngineTestTickActorTickableSlow, STATGROUP_Tickables); }
+};
+
+
+#define LOG_SCOPE_TIME(x) \
+	TRACE_CPUPROFILER_EVENT_SCOPE(x); \
+	FScopeLogTime LogTimePtr(TEXT(#x), nullptr, FScopeLogTime::ScopeLog_Milliseconds)
+
+// Ensures that manually ticking a world works correctly
+IMPLEMENT_CUSTOM_SIMPLE_AUTOMATION_TEST(FBasicTickTest, FEngineTickTestBase, "System.Engine.Tick.BasicTest", EAutomationTestFlags::EditorContext | EAutomationTestFlags::ClientContext | EAutomationTestFlags::ServerContext | EAutomationTestFlags::EngineFilter)
+bool FBasicTickTest::RunTest(const FString& Parameters)
+{
+	int32 ActorCount = 10;
+	int32 TickCount = 10;
+	float DeltaTime = 0.01f;
+
+	if (!CreateTestWorld())
+	{
+		return false;
+	}
+
+	bool bSuccess = true;
+
+	bSuccess &= CreateTestActors(ActorCount, AEngineTestTickActor::StaticClass());
+	bSuccess &= BeginPlayInTestWorld();
+	 
+	if (bSuccess)
+	{
+		for (int32 i = 0; i < TickCount; i++)
+		{
+			TickTestWorld(DeltaTime);
+		}
+
+		CheckTickCount(TEXT("TickCount"), TickCount);
+	}
+
+	// Always reset test world
+	bSuccess &= DestroyTestWorld();
+
+	return bSuccess && !ReportAnyErrors();
+}
+
+// Compares different ways of ticking actors for performance
+IMPLEMENT_CUSTOM_SIMPLE_AUTOMATION_TEST(FPerfTickTest, FEngineTickTestBase, "System.Engine.Tick.PerfTest", EAutomationTestFlags::EditorContext | EAutomationTestFlags::ClientContext | EAutomationTestFlags::ServerContext | EAutomationTestFlags::PerfFilter)
+bool FPerfTickTest::RunTest(const FString& Parameters)
+{
+	int32 ActorCount = 1000;
+	int32 TickCount = 1000; // Increase for one off tests
+	float DeltaTime = 0.01f;
+
+	if (!CreateTestWorld())
+	{
+		ReportAnyErrors();
+		return false;
+	}
+
+	if (BeginPlayInTestWorld())
+	{
+		UE_LOG(LogStats, Log, TEXT("Running FPerfTickTest for %d actors over %d tick frames:"), ActorCount, TickCount);
+
+		{
+			// Time to tick an empty world
+			LOG_SCOPE_TIME(WorldBaseline);
+			for (int32 i = 0; i < TickCount; i++)
+			{
+				TickTestWorld();
+			}
+		}
+
+		if (!CreateTestActors(ActorCount, AEngineTestTickActor::StaticClass()))
+		{
+			return false;
+		}
+		
+		ResetTestActors();
+		{
+			// Tick with normal task graph method
+			LOG_SCOPE_TIME(WorldActorTick);
+			for (int32 i = 0; i < TickCount; i++)
+			{
+				TickTestWorld();
+			}
+		}
+		CheckTickCount(TEXT("WorldActorTick"), TickCount);
+
+
+		FSimpleMulticastDelegate LambdaDelegate, VirtualLambdaDelegate;
+		FSimpleMulticastDelegate UObjectDelegate, VirtualUObjectDelegate;
+		FSimpleMulticastDelegate WeakLambdaDelegate, VirtualWeakLambdaDelegate;
+		FTSTicker TSTicker;
+
+		for (AEngineTestTickActor* TestActor : TestActors)
+		{
+			// Unregister normal ticks
+			TestActor->RegisterAllActorTickFunctions(false, false);
+
+			// Check various delegate types, raw delegates are blocked on UObjects
+			LambdaDelegate.AddLambda([TestActor]() {TestActor->DoTick(); });
+			VirtualLambdaDelegate.AddLambda([TestActor]() {TestActor->VirtualTick(); });
+			UObjectDelegate.AddUObject(TestActor, &AEngineTestTickActor::DoTick);
+			VirtualUObjectDelegate.AddUObject(TestActor, &AEngineTestTickActor::VirtualTick);
+			WeakLambdaDelegate.AddWeakLambda(TestActor, [TestActor]() {TestActor->DoTick(); });
+			VirtualWeakLambdaDelegate.AddWeakLambda(TestActor, [TestActor]() {TestActor->VirtualTick(); });
+			TSTicker.AddTicker(FTickerDelegate::CreateWeakLambda(TestActor, [TestActor](float) {TestActor->VirtualTick(); return true; }), 0.0f);
+		}
+
+		// Possible options for real world ticks
+		ResetTestActors();
+		{
+			LOG_SCOPE_TIME(WorldTSTicker);
+			for (int32 i = 0; i < TickCount; i++)
+			{
+				TickTestWorld();
+				TSTicker.Tick(DeltaTime);
+			}
+		}
+		CheckTickCount(TEXT("WorldTSTicker"), TickCount);
+		TSTicker.Reset();
+
+
+		FTimerManager& TimerManager = GetTestWorld()->GetTimerManager();
+		TArray<FTimerHandle> TimerHandles;
+		TimerHandles.Reserve(TestActors.Num());
+		for (AEngineTestTickActor* TestActor : TestActors)
+		{
+			FTimerHandle& TimerHandle = TimerHandles.AddDefaulted_GetRef();
+			TimerManager.SetTimer(TimerHandle, FTimerDelegate::CreateWeakLambda(TestActor, [TestActor]() {TestActor->VirtualTick(); }), 0.001f,
+				FTimerManagerTimerParameters{ .bLoop = true, .bMaxOncePerFrame = true, .FirstDelay = 0.0f });
+		}
+		
+		// Tick the world once as timers won't tick until the next frame even if they are initialized outside of tick
+		TickTestWorld();
+
+		ResetTestActors();
+		{
+			LOG_SCOPE_TIME(WorldTimerManager);
+			for (int32 i = 0; i < TickCount; i++)
+			{
+				TickTestWorld();
+			}
+		}
+		CheckTickCount(TEXT("WorldTimerManager"), TickCount);
+		for (FTimerHandle& TimerHandle : TimerHandles)
+		{
+			TimerManager.ClearTimer(TimerHandle);
+			ensure(!TimerHandle.IsValid());
+		}
+		TimerHandles.Empty();
+
+		// Fastest possible TickableGameObject
+		TArray<FEngineTestTickActorTickableFast> FastTickables;
+		FastTickables.Reserve(TestActors.Num());
+		for (AEngineTestTickActor* TestActor : TestActors)
+		{
+			FastTickables.Emplace(TestActor);
+		}
+
+		ResetTestActors();
+		{
+			LOG_SCOPE_TIME(WorldTickableFast);
+			for (int32 i = 0; i < TickCount; i++)
+			{
+				TickTestWorld();
+			}
+		}
+		CheckTickCount(TEXT("WorldTickableFast"), TickCount);
+		FastTickables.Empty();
+
+		// Slower unoptimized TickableGameObject
+		TArray<FEngineTestTickActorTickableSlow> SlowTickables;
+		SlowTickables.Reserve(TestActors.Num());
+		for (AEngineTestTickActor* TestActor : TestActors)
+		{
+			SlowTickables.Emplace(TestActor);
+		}
+
+		ResetTestActors();
+		{
+			LOG_SCOPE_TIME(WorldTickableSlow);
+			for (int32 i = 0; i < TickCount; i++)
+			{
+				TickTestWorld();
+			}
+		}
+		CheckTickCount(TEXT("WorldTickableSlow"), TickCount);
+		SlowTickables.Empty();
+
+
+		// Raw function call tests, with a world tick before
+		ResetTestActors();
+		{
+			LOG_SCOPE_TIME(LoopDoTick);
+			for (int32 i = 0; i < TickCount; i++)
+			{
+				TickTestWorld();
+				for (AEngineTestTickActor* TestActor : TestActors)
+				{
+					TestActor->DoTick();
+				}
+			}
+		}
+		CheckTickCount(TEXT("LoopDoTick"), TickCount);
+
+
+		ResetTestActors();
+		{
+			LOG_SCOPE_TIME(LoopVirtualTick);
+			for (int32 i = 0; i < TickCount; i++)
+			{
+				TickTestWorld();
+				for (AEngineTestTickActor* TestActor : TestActors)
+				{
+					TestActor->VirtualTick();
+				}
+			}
+		}
+		CheckTickCount(TEXT("LoopVirtualTick"), TickCount);
+
+
+		ResetTestActors();
+		{
+			FGraphEventRef FakeEvent;
+			LOG_SCOPE_TIME(LoopExecuteTick);
+			for (int32 i = 0; i < TickCount; i++)
+			{
+				TickTestWorld();
+				for (AEngineTestTickActor* TestActor : TestActors)
+				{
+					// TODO could replace with registering a tick manager
+					TestActor->PrimaryActorTick.ExecuteTick(DeltaTime, LEVELTICK_All, ENamedThreads::GameThread, FakeEvent);
+				}
+
+			}
+		}
+		CheckTickCount(TEXT("LoopExecuteTick"), TickCount);
+
+
+		ResetTestActors();
+		{
+			LOG_SCOPE_TIME(LambdaDelegate);
+			for (int32 i = 0; i < TickCount; i++)
+			{
+				TickTestWorld();
+				LambdaDelegate.Broadcast();
+			}
+		}
+		CheckTickCount(TEXT("LambdaDelegate"), TickCount);
+		LambdaDelegate.Clear();
+
+
+		ResetTestActors();
+		{
+			LOG_SCOPE_TIME(VirtualLambdaDelegate);
+			for (int32 i = 0; i < TickCount; i++)
+			{
+				TickTestWorld();
+				VirtualLambdaDelegate.Broadcast();
+			}
+		}
+		CheckTickCount(TEXT("VirtualLambdaDelegate"), TickCount);
+		VirtualLambdaDelegate.Clear();
+
+
+		ResetTestActors();
+		{
+			LOG_SCOPE_TIME(UObjectDelegate);
+			for (int32 i = 0; i < TickCount; i++)
+			{
+				TickTestWorld();
+				UObjectDelegate.Broadcast();
+			}
+		}
+		CheckTickCount(TEXT("UObjectDelegate"), TickCount);
+		UObjectDelegate.Clear();
+
+
+		ResetTestActors();
+		{
+			LOG_SCOPE_TIME(VirtualUObjectDelegate);
+			for (int32 i = 0; i < TickCount; i++)
+			{
+				TickTestWorld();
+				VirtualUObjectDelegate.Broadcast();
+			}
+		}
+		CheckTickCount(TEXT("VirtualUObjectDelegate"), TickCount);
+		VirtualUObjectDelegate.Clear();
+
+
+		ResetTestActors();
+		{
+			LOG_SCOPE_TIME(WeakLambdaDelegate);
+			for (int32 i = 0; i < TickCount; i++)
+			{
+				TickTestWorld();
+				WeakLambdaDelegate.Broadcast();
+			}
+		}
+		CheckTickCount(TEXT("WeakLambdaDelegate"), TickCount);
+		WeakLambdaDelegate.Clear();
+
+
+		ResetTestActors();
+		{
+			LOG_SCOPE_TIME(VirtualWeakLambdaDelegate);
+			for (int32 i = 0; i < TickCount; i++)
+			{
+				TickTestWorld();
+				VirtualWeakLambdaDelegate.Broadcast();
+			}
+		}
+		CheckTickCount(TEXT("VirtualWeakLambdaDelegate"), TickCount);
+		VirtualWeakLambdaDelegate.Clear();
+	}
+	return DestroyTestWorld() && !ReportAnyErrors();
+}
+
+#endif // WITH_AUTOMATION_WORKER
diff --git a/Engine/Plugins/Tests/RuntimeTests/Source/RuntimeTests/Public/EngineRuntimeTests.h b/Engine/Plugins/Tests/RuntimeTests/Source/RuntimeTests/Public/EngineRuntimeTests.h
new file mode 100644
--- /dev/null
+++ b/Engine/Plugins/Tests/RuntimeTests/Source/RuntimeTests/Public/EngineRuntimeTests.h
@@ -0,0 +1,107 @@
+// Copyright Epic Games, Inc. All Rights Reserved.
+#pragma once 
+
+#include "GameFramework/Actor.h"
+#include "Misc/AutomationTest.h"
+#include "Tests/AutomationCommon.h"
+
+#include "EngineRuntimeTests.generated.h"
+
+class UBillboardComponent;
+
+/** A simple actor class that can be manually ticked to test for correctness and performance */
+UCLASS()
+class RUNTIMETESTS_API AEngineTestTickActor : public AActor
+{
+	GENERATED_BODY()
+
+public:
+	AEngineTestTickActor(const FObjectInitializer& ObjectInitializer = FObjectInitializer::Get());
+
+	/** Number of times this has ticked since reset */
+	UPROPERTY(BlueprintReadOnly, Category=Default)
+	int32 TickCount;
+
+	/** If it should actually increase tick count */
+	UPROPERTY(BlueprintReadOnly, Category=Default)
+	bool bShouldIncrementTickCount;
+
+	/** If it should perform other busy work */
+	UPROPERTY(BlueprintReadOnly, Category=Default)
+	bool bShouldDoMath;
+
+	/** Used for bShouldDoMath */
+	UPROPERTY(BlueprintReadOnly, Category=Default)
+	float MathCounter;
+
+	/** Used for bShouldDoMath */
+	UPROPERTY(BlueprintReadOnly, Category=Default)
+	float MathIncrement;
+
+	/** Used for bShouldDoMath */
+	UPROPERTY(BlueprintReadOnly, Category=Default)
+	float MathLimit;
+
+	/** Reset state before next test, call this after unregistering tick */
+	virtual void ResetState();
+
+	/** Do the actual work */
+	void DoTick();
+
+	/** Virtual function wrapper */
+	virtual void VirtualTick();
+
+	/** AActor version */
+	virtual void Tick(float DeltaSeconds) override;
+
+private:
+	UPROPERTY()
+	UBillboardComponent* SpriteComponent;
+};
+
+#if WITH_AUTOMATION_WORKER
+
+/** Automation test base class that wraps a test world and handles checking tick counts */
+class RUNTIMETESTS_API FEngineTickTestBase : public FAutomationTestBase
+{
+public:
+	FEngineTickTestBase(const FString& InName, const bool bInComplexTask);
+
+	virtual ~FEngineTickTestBase();
+
+	/** Gets the world being tested */
+	UWorld* GetTestWorld() const;
+
+	/** Creates a world where actors can be spawned */
+	virtual bool CreateTestWorld();
+
+	/** Spawn actors of subclass */
+	virtual bool CreateTestActors(int32 ActorCount, TSubclassOf<AEngineTestTickActor> ActorClass);
+
+	/** Start play in world, prepare for ticking */
+	virtual bool BeginPlayInTestWorld();
+
+	/** Tick one frame in test world */
+	virtual bool TickTestWorld(float DeltaTime = 0.01f);
+
+	/** Reset the test */
+	virtual bool ResetTestActors();
+
+	/** Checks TickCount on every actor */
+	virtual bool CheckTickCount(const TCHAR* TickTestName, int32 TickCount);
+
+	/** Destroys the test actors */
+	virtual bool DestroyAllTestActors();
+	
+	/** Destroys the test world */
+	virtual bool DestroyTestWorld();
+
+	/** Reports errors to automation system, returns true if there were errors */
+	virtual bool ReportAnyErrors();
+
+protected:
+	TUniquePtr<FTestWorldWrapper> WorldWrapper;
+	TArray<AEngineTestTickActor*> TestActors;
+};
+
+#endif // WITH_AUTOMATION_WORKER
\ No newline at end of file
