Add System.Engine.Tick.OrderTest automated test to validate tick ordering and added some cvars for tweaking behavior of System.Engine.Tick.PerfTest

diff --git a/Engine/Plugins/Tests/RuntimeTests/Source/RuntimeTests/Private/EngineRuntimeTests.cpp b/Engine/Plugins/Tests/RuntimeTests/Source/RuntimeTests/Private/EngineRuntimeTests.cpp
--- a/Engine/Plugins/Tests/RuntimeTests/Source/RuntimeTests/Private/EngineRuntimeTests.cpp
+++ b/Engine/Plugins/Tests/RuntimeTests/Source/RuntimeTests/Private/EngineRuntimeTests.cpp
@@ -7,6 +7,7 @@
 #include "Engine/Level.h"
 #include "Engine/World.h"
 #include "EngineUtils.h"
+#include "Math/RandomStream.h"
 #include "Misc/AutomationTest.h"
 #include "Stats/StatsMisc.h"
 #include "Containers/Ticker.h"
@@ -14,6 +15,8 @@
 #include "TimerManager.h"
 #include "Tickable.h"
 
+int32 AEngineTestTickActor::CurrentTickOrder = 0;
+
 AEngineTestTickActor::AEngineTestTickActor(const FObjectInitializer& ObjectInitializer)
 	: Super(ObjectInitializer)
 {
@@ -24,6 +27,7 @@ AEngineTestTickActor::AEngineTestTickActor(const FObjectInitializer& ObjectIniti
 		RootComponent = SpriteComponent;
 	}
 
+	PrimaryActorTick.TickGroup = TG_PrePhysics;
 	PrimaryActorTick.bCanEverTick = true;
 
 	ResetState();
@@ -32,6 +36,7 @@ AEngineTestTickActor::AEngineTestTickActor(const FObjectInitializer& ObjectIniti
 void AEngineTestTickActor::ResetState()
 {
 	TickCount = 0;
+	TickOrder = 0;
 	bShouldIncrementTickCount = true;
 	bShouldDoMath = true;
 	MathCounter = 0.0f;
@@ -54,6 +59,8 @@ void AEngineTestTickActor::DoTick()
 			MathCounter += MathIncrement;
 		}
 	}
+
+	TickOrder = CurrentTickOrder++;
 }
 
 void AEngineTestTickActor::VirtualTick()
@@ -145,6 +152,7 @@ bool FEngineTickTestBase::TickTestWorld(float DeltaTime)
 		return false;
 	}
 
+	AEngineTestTickActor::CurrentTickOrder = 1;
 	return WorldWrapper->TickTestWorld(DeltaTime);
 }
 
@@ -288,12 +296,124 @@ bool FBasicTickTest::RunTest(const FString& Parameters)
 	return bSuccess && !ReportAnyErrors();
 }
 
+// Verify different methods of ordering ticks
+IMPLEMENT_CUSTOM_SIMPLE_AUTOMATION_TEST(FOrderTickTest, FEngineTickTestBase, "System.Engine.Tick.OrderTest", EAutomationTestFlags::EditorContext | EAutomationTestFlags::ClientContext | EAutomationTestFlags::ServerContext | EAutomationTestFlags::EngineFilter)
+bool FOrderTickTest::RunTest(const FString& Parameters)
+{
+	float DeltaTime = 0.01f;
+	int32 ActorCount = 1000;
+
+	if (!CreateTestWorld())
+	{
+		return false;
+	}
+
+	bool bSuccess = true;
+
+	bSuccess &= CreateTestActors(ActorCount, AEngineTestTickActor::StaticClass());
+	bSuccess &= BeginPlayInTestWorld();
+
+	if (bSuccess)
+	{
+		check(TestActors.Num() == ActorCount);
+
+		// Semirandom numbers, generally tick happens based on order of spawn but that is not guaranteed
+		AEngineTestTickActor* HighPriority = TestActors[12];
+		AEngineTestTickActor* HighPrereq = TestActors[18];
+		AEngineTestTickActor* PostPhysics = TestActors[2];
+		AEngineTestTickActor* PostPhysicsDep = TestActors[75];
+		AEngineTestTickActor* PosyPhysicsDep2 = TestActors[45];
+		AEngineTestTickActor* TickInterval = TestActors[32];
+		AEngineTestTickActor* TickIntervalDep = TestActors[23];
+
+		HighPriority->PrimaryActorTick.SetPriorityIncludingPrerequisites(true);
+		PostPhysics->PrimaryActorTick.TickGroup = TG_PostPhysics;
+
+		ResetTestActors();
+		TickTestWorld(DeltaTime);
+
+		TestEqual(TEXT("HighPriority tickorder"), HighPriority->TickOrder, 1);
+		TestEqual(TEXT("PostPhysics tickorder"), PostPhysics->TickOrder, ActorCount);
+
+		HighPriority->AddTickPrerequisiteActor(HighPrereq);
+
+		// This has to be refreshed the tick prereq is set right now, comment out to verify
+		HighPriority->PrimaryActorTick.SetPriorityIncludingPrerequisites(false);
+		HighPriority->PrimaryActorTick.SetPriorityIncludingPrerequisites(true);
+
+		// Test dependency group demoting
+		PostPhysicsDep->AddTickPrerequisiteActor(PostPhysics);
+
+		ResetTestActors();
+		TickTestWorld(DeltaTime);
+
+		TestEqual(TEXT("HighPrereq tickorder"), HighPrereq->TickOrder, 1);
+		TestEqual(TEXT("HighPriority tickorder"), HighPriority->TickOrder, 2);
+		TestEqual(TEXT("PostPhysicsDep tickorder"), PostPhysicsDep->TickOrder, ActorCount);
+
+
+		// Uncomment to test circular reference, which throws off ordering
+		// PostPhysics->AddTickPrerequisiteActor(PosyPhysicsDep2);
+		PosyPhysicsDep2->AddTickPrerequisiteActor(PostPhysicsDep);
+
+		// Test tick interval, it will be run the first tick but not the second
+		TickInterval->SetActorTickInterval(0.5f);
+		TickInterval->PrimaryActorTick.TickGroup = TG_PostUpdateWork;
+
+		// The dependency will be respected the first time, but not the second
+		TickIntervalDep->AddTickPrerequisiteActor(TickInterval);
+
+		ResetTestActors();
+		TickTestWorld(DeltaTime);
+
+		TestEqual(TEXT("TickInterval count"), TickInterval->TickCount, 1);
+		TestEqual(TEXT("TickIntervalDep tickorder"), TickIntervalDep->TickOrder, ActorCount); // This will be last because dependency is respected
+
+		TickTestWorld(DeltaTime);
+
+		TestEqual(TEXT("TickInterval count"), TickInterval->TickCount, 1); // This was skipped by second tick
+		TestEqual(TEXT("TickIntervalDep count"), TickIntervalDep->TickCount, 2);
+		TestEqual(TEXT("PosyPhysicsDep2 tickorder"), PosyPhysicsDep2->TickOrder, ActorCount - 1); // TickInterval is skipped on the second frame so this is last of 99
+
+		TestEqual(TEXT("HighPrereq tickorder"), HighPrereq->TickOrder, 1);
+		TestEqual(TEXT("HighPriority tickorder"), HighPriority->TickOrder, 2);
+	}
+
+	// Always reset test world
+	bSuccess &= DestroyTestWorld();
+
+	return bSuccess && !ReportAnyErrors();
+}
+
+
+static TAutoConsoleVariable<int32> CVarEngineTickPerfOptions(
+	TEXT("Automation.Test.EngineTickPerf.Options"),
+	0,
+	TEXT("Bitfield to modify options used for tick test.\n")
+	TEXT("0 - No tick dependencies or intervals\n")
+	TEXT("1 - Add tick dependencies\n")
+	TEXT("2 - Add tick intervals\n")
+	TEXT("3 - Add tick dependencies and intervals\n"),
+	ECVF_Default);
+
+static TAutoConsoleVariable<int32> CVarEngineTickPerfActorCount(
+	TEXT("Automation.Test.EngineTickPerf.ActorCount"),
+	1000,
+	TEXT("Number of actors to spawn for tick test\n"),
+	ECVF_Default);
+
+static TAutoConsoleVariable<int32> CVarEngineTickPerfTickCount(
+	TEXT("Automation.Test.EngineTickPerf.TickCount"),
+	1000,
+	TEXT("Number of frames to tick\n"),
+	ECVF_Default);
+
 // Compares different ways of ticking actors for performance
 IMPLEMENT_CUSTOM_SIMPLE_AUTOMATION_TEST(FPerfTickTest, FEngineTickTestBase, "System.Engine.Tick.PerfTest", EAutomationTestFlags::EditorContext | EAutomationTestFlags::ClientContext | EAutomationTestFlags::ServerContext | EAutomationTestFlags::PerfFilter)
 bool FPerfTickTest::RunTest(const FString& Parameters)
 {
-	int32 ActorCount = 1000;
-	int32 TickCount = 1000; // Increase for one off tests
+	const int32 ActorCount = CVarEngineTickPerfActorCount->GetInt();
+	const int32 TickCount = CVarEngineTickPerfTickCount->GetInt();
 	float DeltaTime = 0.01f;
 
 	if (!CreateTestWorld())
@@ -319,7 +439,28 @@ bool FPerfTickTest::RunTest(const FString& Parameters)
 		{
 			return false;
 		}
-		
+
+		const int32 TestOptions = CVarEngineTickPerfOptions->GetInt();
+		const int32 RandomSeed = 0xABCD1234;
+		FRandomStream RandomSource(RandomSeed);
+
+		// Add some semi-random timing and dependency changes
+		for (int32 i = 0; i < ActorCount; i++)
+		{
+			if ((TestOptions & 0x00000001) != 0 && i != (ActorCount-1))
+			{
+				// Enable dependencies on a random later actor
+				TestActors[i]->AddTickPrerequisiteActor(TestActors[RandomSource.RandRange(i+1, ActorCount-1)]);
+				// TestActors[i]->AddTickPrerequisiteActor(TestActors[RandomSource.RandHelper(ActorCount)]); // This creates infinite loops which can deadlock the engine
+			}
+
+			if ((TestOptions & 0x00000002) != 0)
+			{
+				// Enable a small interval, this should not affect actual timing
+				TestActors[i]->SetActorTickInterval(DeltaTime / 2 + RandomSource.FRandRange(-DeltaTime/10, DeltaTime/10));
+			}
+		}
+
 		ResetTestActors();
 		{
 			// Tick with normal task graph method
diff --git a/Engine/Plugins/Tests/RuntimeTests/Source/RuntimeTests/Public/EngineRuntimeTests.h b/Engine/Plugins/Tests/RuntimeTests/Source/RuntimeTests/Public/EngineRuntimeTests.h
--- a/Engine/Plugins/Tests/RuntimeTests/Source/RuntimeTests/Public/EngineRuntimeTests.h
+++ b/Engine/Plugins/Tests/RuntimeTests/Source/RuntimeTests/Public/EngineRuntimeTests.h
@@ -21,9 +21,16 @@ public:
 	/** Number of times this has ticked since reset */
 	UPROPERTY(BlueprintReadOnly, Category=Default)
 	int32 TickCount;
+	
+	/** Indicates when this was ticked in a frame, with 1 being first */
+	UPROPERTY(BlueprintReadOnly, Category = Default)
+	int32 TickOrder;
+
+	/** Used to set TickOrder, reset to 1 at the start of every frame */
+	static int32 CurrentTickOrder;
 
 	/** If it should actually increase tick count */
-	UPROPERTY(BlueprintReadOnly, Category=Default)
+	UPROPERTY(BlueprintReadOnly, Category = Default)
 	bool bShouldIncrementTickCount;
 
 	/** If it should perform other busy work */
