Add task subtests to tick performance tests to isolate task and tick overhead, best compared with "tick.AllowBatchedTicks 0" to disable batching
Add "Automation.Test.EngineTickPerf.RunTests Test1 Test2" cvar to make it easier to profile specific subtests, and it will skip a few less useful ones by default

diff --git a/Engine/Plugins/Tests/RuntimeTests/Source/RuntimeTests/Private/EngineRuntimeTests.cpp b/Engine/Plugins/Tests/RuntimeTests/Source/RuntimeTests/Private/EngineRuntimeTests.cpp
--- a/Engine/Plugins/Tests/RuntimeTests/Source/RuntimeTests/Private/EngineRuntimeTests.cpp
+++ b/Engine/Plugins/Tests/RuntimeTests/Source/RuntimeTests/Private/EngineRuntimeTests.cpp
@@ -14,6 +14,7 @@
 #include "Tests/AutomationCommon.h"
 #include "TimerManager.h"
 #include "Tickable.h"
+#include "Tasks/Task.h"
 
 int32 AEngineTestTickActor::CurrentTickOrder = 0;
 
@@ -408,6 +409,44 @@ static TAutoConsoleVariable<int32> CVarEngineTickPerfTickCount(
 	TEXT("Number of frames to tick\n"),
 	ECVF_Default);
 
+static TAutoConsoleVariable<FString> CVarEngineTickPerfRunTests(
+	TEXT("Automation.Test.EngineTickPerf.RunTests"),
+	TEXT(""),
+	TEXT("Specific tests to run separated by spaces. If empty it will run all the named tests that pass true to IsTestEnabled()\n"),
+	ECVF_Default);
+
+// Simplest task graph task to eliminate tick/safety overhead
+class FEngineTestTickSimpleTask
+{
+public:
+	AEngineTestTickActor* Actor = nullptr;
+	ENamedThreads::Type DesiredThread = ENamedThreads::GameThread;
+
+	FORCEINLINE FEngineTestTickSimpleTask(AEngineTestTickActor* InActor, ENamedThreads::Type InDesiredThread)
+		: Actor(InActor)
+		, DesiredThread(InDesiredThread)
+	{
+	}
+
+	// Functions used by TGraphTask to schedule and execute task
+	static FORCEINLINE TStatId GetStatId()
+	{
+		RETURN_QUICK_DECLARE_CYCLE_STAT(FEngineTestTickSimpleTask, STATGROUP_TaskGraphTasks);
+	}
+	FORCEINLINE ENamedThreads::Type GetDesiredThread()
+	{
+		return DesiredThread;
+	}
+	static FORCEINLINE ESubsequentsMode::Type GetSubsequentsMode()
+	{
+		return ESubsequentsMode::TrackSubsequents;
+	}
+	void DoTask(ENamedThreads::Type CurrentThread, const FGraphEventRef& MyCompletionGraphEvent)
+	{
+		Actor->VirtualTick();
+	}
+};
+
 // Compares different ways of ticking actors for performance
 IMPLEMENT_CUSTOM_SIMPLE_AUTOMATION_TEST(FPerfTickTest, FEngineTickTestBase, "System.Engine.Tick.PerfTest", EAutomationTestFlags::EditorContext | EAutomationTestFlags::ClientContext | EAutomationTestFlags::ServerContext | EAutomationTestFlags::PerfFilter)
 bool FPerfTickTest::RunTest(const FString& Parameters)
@@ -415,6 +454,17 @@ bool FPerfTickTest::RunTest(const FString& Parameters)
 	const int32 ActorCount = CVarEngineTickPerfActorCount->GetInt();
 	const int32 TickCount = CVarEngineTickPerfTickCount->GetInt();
 	float DeltaTime = 0.01f;
+	TArray<FString> TestsToRun;
+	CVarEngineTickPerfRunTests->GetString().ParseIntoArray(TestsToRun, TEXT(" "));
+
+	auto IsTestEnabled = [TestsToRun](FString InString, bool bEnableByDefault = true)
+	{
+		if (TestsToRun.Contains(InString) || (bEnableByDefault && TestsToRun.Num() == 0))
+		{
+			return true;
+		}
+		return false;
+	};
 
 	if (!CreateTestWorld())
 	{
@@ -461,17 +511,19 @@ bool FPerfTickTest::RunTest(const FString& Parameters)
 			}
 		}
 
-		ResetTestActors();
+		if (IsTestEnabled(TEXT("WorldActorTick"), true))
 		{
-			// Tick with normal task graph method
-			LOG_SCOPE_TIME(WorldActorTick);
-			for (int32 i = 0; i < TickCount; i++)
+			ResetTestActors();
 			{
-				TickTestWorld();
+				// Tick with normal task graph method
+				LOG_SCOPE_TIME(WorldActorTick);
+				for (int32 i = 0; i < TickCount; i++)
+				{
+					TickTestWorld();
+				}
 			}
+			CheckTickCount(TEXT("WorldActorTick"), TickCount);
 		}
-		CheckTickCount(TEXT("WorldActorTick"), TickCount);
-
 
 		FSimpleMulticastDelegate LambdaDelegate, VirtualLambdaDelegate;
 		FSimpleMulticastDelegate UObjectDelegate, VirtualUObjectDelegate;
@@ -490,216 +542,321 @@ bool FPerfTickTest::RunTest(const FString& Parameters)
 			VirtualUObjectDelegate.AddUObject(TestActor, &AEngineTestTickActor::VirtualTick);
 			WeakLambdaDelegate.AddWeakLambda(TestActor, [TestActor]() {TestActor->DoTick(); });
 			VirtualWeakLambdaDelegate.AddWeakLambda(TestActor, [TestActor]() {TestActor->VirtualTick(); });
-			TSTicker.AddTicker(FTickerDelegate::CreateWeakLambda(TestActor, [TestActor](float) {TestActor->VirtualTick(); return true; }), 0.0f);
 		}
 
-		// Possible options for real world ticks
-		ResetTestActors();
+		if (IsTestEnabled(TEXT("TaskGraph"), false))
 		{
-			LOG_SCOPE_TIME(WorldTSTicker);
-			for (int32 i = 0; i < TickCount; i++)
+			// Using task graph directly to avoid tick overhead, this is the same process the tick manager uses with construct and dispatch
+			FGraphEventArray GraphEvents;
+			GraphEvents.Reserve(TestActors.Num());
+
+			ResetTestActors();
 			{
-				TickTestWorld();
-				TSTicker.Tick(DeltaTime);
+				LOG_SCOPE_TIME(TaskGraphConstructAndWait);
+				for (int32 i = 0; i < TickCount; i++)
+				{
+					for (AEngineTestTickActor* TestActor : TestActors)
+					{
+						TGraphTask<FEngineTestTickSimpleTask>* TaskPtr = TGraphTask<FEngineTestTickSimpleTask>::CreateTask(nullptr, ENamedThreads::GameThread).ConstructAndHold(TestActor, ENamedThreads::GameThread);
+						GraphEvents.Emplace(TaskPtr->GetCompletionEvent());
+					}
+
+					// Separate loop as this is what tick does
+					for (FGraphEventRef& GraphEventRef : GraphEvents)
+					{
+						GraphEventRef->Unlock();
+					}
+
+					FTaskGraphInterface::Get().ProcessUntilTasksComplete(GraphEvents, ENamedThreads::GameThread);
+					GraphEvents.Reset();
+
+					// Do an empty world tick to match baseline
+					TickTestWorld();
+				}
 			}
+			CheckTickCount(TEXT("TaskGraph"), TickCount);
 		}
-		CheckTickCount(TEXT("WorldTSTicker"), TickCount);
-		TSTicker.Reset();
-
 
-		FTimerManager& TimerManager = GetTestWorld()->GetTimerManager();
-		TArray<FTimerHandle> TimerHandles;
-		TimerHandles.Reserve(TestActors.Num());
-		for (AEngineTestTickActor* TestActor : TestActors)
+		if (IsTestEnabled(TEXT("BaseTask"), false))
 		{
-			FTimerHandle& TimerHandle = TimerHandles.AddDefaulted_GetRef();
-			TimerManager.SetTimer(TimerHandle, FTimerDelegate::CreateWeakLambda(TestActor, [TestActor]() {TestActor->VirtualTick(); }), 0.001f,
-				FTimerManagerTimerParameters{ .bLoop = true, .bMaxOncePerFrame = true, .FirstDelay = 0.0f });
-		}
-		
-		// Tick the world once as timers won't tick until the next frame even if they are initialized outside of tick
-		TickTestWorld();
+			// Use tasks api directly, but still on game thread. This is not a normal workflow but is useful for isolated profiling
+			TArray<UE::Tasks::FTask> TaskArray;
+			TaskArray.Reserve(TestActors.Num());
 
-		ResetTestActors();
-		{
-			LOG_SCOPE_TIME(WorldTimerManager);
-			for (int32 i = 0; i < TickCount; i++)
+			ResetTestActors();
 			{
-				TickTestWorld();
+				LOG_SCOPE_TIME(BaseTask);
+				for (int32 i = 0; i < TickCount; i++)
+				{
+					for (AEngineTestTickActor* TestActor : TestActors)
+					{
+						TaskArray.Add(UE::Tasks::Launch(UE_SOURCE_LOCATION,
+							[TestActor] {
+								TestActor->VirtualTick();
+							},
+							UE::Tasks::ETaskPriority::Normal,
+							UE::Tasks::EExtendedTaskPriority::GameThreadNormalPri
+						));
+					}
+
+					UE::Tasks::FTask ReturnTask = UE::Tasks::Launch(UE_SOURCE_LOCATION,
+						[] {},
+						TaskArray,
+						UE::Tasks::ETaskPriority::Normal,
+						UE::Tasks::EExtendedTaskPriority::GameThreadNormalPri);
+
+					ReturnTask.Wait();
+					TaskArray.Reset();
+
+					// Do an empty world tick to match baseline
+					TickTestWorld();
+				}
 			}
+			CheckTickCount(TEXT("BaseTask"), TickCount);
 		}
-		CheckTickCount(TEXT("WorldTimerManager"), TickCount);
-		for (FTimerHandle& TimerHandle : TimerHandles)
-		{
-			TimerManager.ClearTimer(TimerHandle);
-			ensure(!TimerHandle.IsValid());
-		}
-		TimerHandles.Empty();
 
-		// Fastest possible TickableGameObject
-		TArray<FEngineTestTickActorTickableFast> FastTickables;
-		FastTickables.Reserve(TestActors.Num());
-		for (AEngineTestTickActor* TestActor : TestActors)
+		if (IsTestEnabled(TEXT("WorldTSTicker"), true))
 		{
-			FastTickables.Emplace(TestActor);
+			for (AEngineTestTickActor* TestActor : TestActors)
+			{
+				TSTicker.AddTicker(FTickerDelegate::CreateWeakLambda(TestActor, [TestActor](float) {TestActor->VirtualTick(); return true; }), 0.0f);
+			}
+
+			ResetTestActors();
+			{
+				LOG_SCOPE_TIME(WorldTSTicker);
+				for (int32 i = 0; i < TickCount; i++)
+				{
+					TickTestWorld();
+					TSTicker.Tick(DeltaTime);
+				}
+			}
+			CheckTickCount(TEXT("WorldTSTicker"), TickCount);
 		}
+		TSTicker.Reset();
 
-		ResetTestActors();
+		if (IsTestEnabled(TEXT("WorldTimerManager"), true))
 		{
-			LOG_SCOPE_TIME(WorldTickableFast);
-			for (int32 i = 0; i < TickCount; i++)
+			FTimerManager& TimerManager = GetTestWorld()->GetTimerManager();
+			TArray<FTimerHandle> TimerHandles;
+			TimerHandles.Reserve(TestActors.Num());
+			for (AEngineTestTickActor* TestActor : TestActors)
 			{
-				TickTestWorld();
+				FTimerHandle& TimerHandle = TimerHandles.AddDefaulted_GetRef();
+				TimerManager.SetTimer(TimerHandle, FTimerDelegate::CreateWeakLambda(TestActor, [TestActor]() {TestActor->VirtualTick(); }), 0.001f,
+					FTimerManagerTimerParameters{ .bLoop = true, .bMaxOncePerFrame = true, .FirstDelay = 0.0f });
+			}
+		
+			// Tick the world once as timers won't tick until the next frame even if they are initialized outside of tick
+			TickTestWorld();
+
+			ResetTestActors();
+			{
+				LOG_SCOPE_TIME(WorldTimerManager);
+				for (int32 i = 0; i < TickCount; i++)
+				{
+					TickTestWorld();
+				}
+			}
+			CheckTickCount(TEXT("WorldTimerManager"), TickCount);
+			for (FTimerHandle& TimerHandle : TimerHandles)
+			{
+				TimerManager.ClearTimer(TimerHandle);
+				ensure(!TimerHandle.IsValid());
 			}
+			TimerHandles.Empty();
 		}
-		CheckTickCount(TEXT("WorldTickableFast"), TickCount);
-		FastTickables.Empty();
 
-		// Slower unoptimized TickableGameObject
-		TArray<FEngineTestTickActorTickableSlow> SlowTickables;
-		SlowTickables.Reserve(TestActors.Num());
-		for (AEngineTestTickActor* TestActor : TestActors)
+		if (IsTestEnabled(TEXT("WorldTickableFast"), true))
 		{
-			SlowTickables.Emplace(TestActor);
+			// Fastest possible TickableGameObject
+			TArray<FEngineTestTickActorTickableFast> FastTickables;
+			FastTickables.Reserve(TestActors.Num());
+			for (AEngineTestTickActor* TestActor : TestActors)
+			{
+				FastTickables.Emplace(TestActor);
+			}
+
+			ResetTestActors();
+			{
+				LOG_SCOPE_TIME(WorldTickableFast);
+				for (int32 i = 0; i < TickCount; i++)
+				{
+					TickTestWorld();
+				}
+			}
+			CheckTickCount(TEXT("WorldTickableFast"), TickCount);
+			FastTickables.Empty();
 		}
 
-		ResetTestActors();
+		if (IsTestEnabled(TEXT("WorldTickableSlow"), false))
 		{
-			LOG_SCOPE_TIME(WorldTickableSlow);
-			for (int32 i = 0; i < TickCount; i++)
+			// Slower unoptimized TickableGameObject
+			TArray<FEngineTestTickActorTickableSlow> SlowTickables;
+			SlowTickables.Reserve(TestActors.Num());
+			for (AEngineTestTickActor* TestActor : TestActors)
 			{
-				TickTestWorld();
+				SlowTickables.Emplace(TestActor);
+			}
+
+			ResetTestActors();
+			{
+				LOG_SCOPE_TIME(WorldTickableSlow);
+				for (int32 i = 0; i < TickCount; i++)
+				{
+					TickTestWorld();
+				}
 			}
+			CheckTickCount(TEXT("WorldTickableSlow"), TickCount);
+			SlowTickables.Empty();
 		}
-		CheckTickCount(TEXT("WorldTickableSlow"), TickCount);
-		SlowTickables.Empty();
 
 
-		// Raw function call tests, with a world tick before
-		ResetTestActors();
+
+		if (IsTestEnabled(TEXT("LoopDoTick"), true))
 		{
-			LOG_SCOPE_TIME(LoopDoTick);
-			for (int32 i = 0; i < TickCount; i++)
+			// Raw function call tests, with a world tick before
+			ResetTestActors();
 			{
-				TickTestWorld();
-				for (AEngineTestTickActor* TestActor : TestActors)
+				LOG_SCOPE_TIME(LoopDoTick);
+				for (int32 i = 0; i < TickCount; i++)
 				{
-					TestActor->DoTick();
+					TickTestWorld();
+					for (AEngineTestTickActor* TestActor : TestActors)
+					{
+						TestActor->DoTick();
+					}
 				}
 			}
+			CheckTickCount(TEXT("LoopDoTick"), TickCount);
 		}
-		CheckTickCount(TEXT("LoopDoTick"), TickCount);
-
 
-		ResetTestActors();
+		if (IsTestEnabled(TEXT("LoopVirtualTick"), false))
 		{
-			LOG_SCOPE_TIME(LoopVirtualTick);
-			for (int32 i = 0; i < TickCount; i++)
+			ResetTestActors();
 			{
-				TickTestWorld();
-				for (AEngineTestTickActor* TestActor : TestActors)
+				LOG_SCOPE_TIME(LoopVirtualTick);
+				for (int32 i = 0; i < TickCount; i++)
 				{
-					TestActor->VirtualTick();
+					TickTestWorld();
+					for (AEngineTestTickActor* TestActor : TestActors)
+					{
+						TestActor->VirtualTick();
+					}
 				}
 			}
+			CheckTickCount(TEXT("LoopVirtualTick"), TickCount);
 		}
-		CheckTickCount(TEXT("LoopVirtualTick"), TickCount);
-
 
-		ResetTestActors();
+		if (IsTestEnabled(TEXT("LoopExecuteTick"), true))
 		{
-			FGraphEventRef FakeEvent;
-			LOG_SCOPE_TIME(LoopExecuteTick);
-			for (int32 i = 0; i < TickCount; i++)
+			ResetTestActors();
 			{
-				TickTestWorld();
-				for (AEngineTestTickActor* TestActor : TestActors)
+				FGraphEventRef FakeEvent;
+				LOG_SCOPE_TIME(LoopExecuteTick);
+				for (int32 i = 0; i < TickCount; i++)
 				{
-					// TODO could replace with registering a tick manager
-					TestActor->PrimaryActorTick.ExecuteTick(DeltaTime, LEVELTICK_All, ENamedThreads::GameThread, FakeEvent);
-				}
+					TickTestWorld();
+					for (AEngineTestTickActor* TestActor : TestActors)
+					{
+						// TODO could replace with registering a tick manager
+						TestActor->PrimaryActorTick.ExecuteTick(DeltaTime, LEVELTICK_All, ENamedThreads::GameThread, FakeEvent);
+					}
 
+				}
 			}
+			CheckTickCount(TEXT("LoopExecuteTick"), TickCount);
 		}
-		CheckTickCount(TEXT("LoopExecuteTick"), TickCount);
-
 
-		ResetTestActors();
+		if (IsTestEnabled(TEXT("LambdaDelegate"), true))
 		{
-			LOG_SCOPE_TIME(LambdaDelegate);
-			for (int32 i = 0; i < TickCount; i++)
+			ResetTestActors();
 			{
-				TickTestWorld();
-				LambdaDelegate.Broadcast();
+				LOG_SCOPE_TIME(LambdaDelegate);
+				for (int32 i = 0; i < TickCount; i++)
+				{
+					TickTestWorld();
+					LambdaDelegate.Broadcast();
+				}
 			}
+			CheckTickCount(TEXT("LambdaDelegate"), TickCount);
+			LambdaDelegate.Clear();
 		}
-		CheckTickCount(TEXT("LambdaDelegate"), TickCount);
-		LambdaDelegate.Clear();
 
-
-		ResetTestActors();
+		if (IsTestEnabled(TEXT("VirtualLambdaDelegate"), false))
 		{
-			LOG_SCOPE_TIME(VirtualLambdaDelegate);
-			for (int32 i = 0; i < TickCount; i++)
+			ResetTestActors();
 			{
-				TickTestWorld();
-				VirtualLambdaDelegate.Broadcast();
+				LOG_SCOPE_TIME(VirtualLambdaDelegate);
+				for (int32 i = 0; i < TickCount; i++)
+				{
+					TickTestWorld();
+					VirtualLambdaDelegate.Broadcast();
+				}
 			}
+			CheckTickCount(TEXT("VirtualLambdaDelegate"), TickCount);
+			VirtualLambdaDelegate.Clear();
 		}
-		CheckTickCount(TEXT("VirtualLambdaDelegate"), TickCount);
-		VirtualLambdaDelegate.Clear();
 
-
-		ResetTestActors();
+		if (IsTestEnabled(TEXT("UObjectDelegate"), false))
 		{
-			LOG_SCOPE_TIME(UObjectDelegate);
-			for (int32 i = 0; i < TickCount; i++)
+			ResetTestActors();
 			{
-				TickTestWorld();
-				UObjectDelegate.Broadcast();
+				LOG_SCOPE_TIME(UObjectDelegate);
+				for (int32 i = 0; i < TickCount; i++)
+				{
+					TickTestWorld();
+					UObjectDelegate.Broadcast();
+				}
 			}
+			CheckTickCount(TEXT("UObjectDelegate"), TickCount);
+			UObjectDelegate.Clear();
 		}
-		CheckTickCount(TEXT("UObjectDelegate"), TickCount);
-		UObjectDelegate.Clear();
 
-
-		ResetTestActors();
+		if (IsTestEnabled(TEXT("VirtualUObjectDelegate"), false))
 		{
-			LOG_SCOPE_TIME(VirtualUObjectDelegate);
-			for (int32 i = 0; i < TickCount; i++)
+			ResetTestActors();
 			{
-				TickTestWorld();
-				VirtualUObjectDelegate.Broadcast();
+				LOG_SCOPE_TIME(VirtualUObjectDelegate);
+				for (int32 i = 0; i < TickCount; i++)
+				{
+					TickTestWorld();
+					VirtualUObjectDelegate.Broadcast();
+				}
 			}
+			CheckTickCount(TEXT("VirtualUObjectDelegate"), TickCount);
+			VirtualUObjectDelegate.Clear();
 		}
-		CheckTickCount(TEXT("VirtualUObjectDelegate"), TickCount);
-		VirtualUObjectDelegate.Clear();
 
-
-		ResetTestActors();
+		if (IsTestEnabled(TEXT("WeakLambdaDelegate"), false))
 		{
-			LOG_SCOPE_TIME(WeakLambdaDelegate);
-			for (int32 i = 0; i < TickCount; i++)
+			ResetTestActors();
 			{
-				TickTestWorld();
-				WeakLambdaDelegate.Broadcast();
+				LOG_SCOPE_TIME(WeakLambdaDelegate);
+				for (int32 i = 0; i < TickCount; i++)
+				{
+					TickTestWorld();
+					WeakLambdaDelegate.Broadcast();
+				}
 			}
+			CheckTickCount(TEXT("WeakLambdaDelegate"), TickCount);
+			WeakLambdaDelegate.Clear();
 		}
-		CheckTickCount(TEXT("WeakLambdaDelegate"), TickCount);
-		WeakLambdaDelegate.Clear();
 
-
-		ResetTestActors();
+		if (IsTestEnabled(TEXT("VirtualWeakLambdaDelegate"), true))
 		{
-			LOG_SCOPE_TIME(VirtualWeakLambdaDelegate);
-			for (int32 i = 0; i < TickCount; i++)
+			ResetTestActors();
 			{
-				TickTestWorld();
-				VirtualWeakLambdaDelegate.Broadcast();
+				LOG_SCOPE_TIME(VirtualWeakLambdaDelegate);
+				for (int32 i = 0; i < TickCount; i++)
+				{
+					TickTestWorld();
+					VirtualWeakLambdaDelegate.Broadcast();
+				}
 			}
+			CheckTickCount(TEXT("VirtualWeakLambdaDelegate"), TickCount);
+			VirtualWeakLambdaDelegate.Clear();
 		}
-		CheckTickCount(TEXT("VirtualWeakLambdaDelegate"), TickCount);
-		VirtualWeakLambdaDelegate.Clear();
 	}
 	return DestroyTestWorld() && !ReportAnyErrors();
 }
