Add additional runtime tests for Mass, the async message system, and the new TaskSyncManager, all currently disabled by default

diff --git a/Engine/Plugins/Tests/RuntimeTests/RuntimeTests.uplugin b/Engine/Plugins/Tests/RuntimeTests/RuntimeTests.uplugin
--- a/Engine/Plugins/Tests/RuntimeTests/RuntimeTests.uplugin
+++ b/Engine/Plugins/Tests/RuntimeTests/RuntimeTests.uplugin
@@ -20,5 +20,11 @@
 			"Type": "Runtime",
 			"LoadingPhase": "PreDefault"
 		}
+	],
+	"Plugins": [
+		{
+			"Name": "AsyncMessageSystem",
+			"Enabled": true
+		}
 	]
 }
\ No newline at end of file
diff --git a/Engine/Plugins/Tests/RuntimeTests/Source/RuntimeTests/Private/EngineRuntimeTests.cpp b/Engine/Plugins/Tests/RuntimeTests/Source/RuntimeTests/Private/EngineRuntimeTests.cpp
--- a/Engine/Plugins/Tests/RuntimeTests/Source/RuntimeTests/Private/EngineRuntimeTests.cpp
+++ b/Engine/Plugins/Tests/RuntimeTests/Source/RuntimeTests/Private/EngineRuntimeTests.cpp
@@ -14,7 +14,19 @@
 #include "Tests/AutomationCommon.h"
 #include "TimerManager.h"
 #include "Tickable.h"
+#include "AsyncMessageHandle.h"
+#include "AsyncMessageSystemBase.h"
+#include "AsyncMessageWorldSubsystem.h"
+#include "AsyncGameplayMessageSystem.h"
+#include "NativeGameplayTags.h"
 #include "Tasks/Task.h"
+#include "MassEntityManager.h"
+#include "MassEntityUtils.h"
+#include "MassProcessingTypes.h"
+#include "MassExecutor.h"
+#include "MassExecutionContext.h"
+#include "MassProcessingPhaseManager.h"
+#include "TaskSyncManager.h"
 
 int32 AEngineTestTickActor::CurrentTickOrder = 0;
 
@@ -386,6 +398,242 @@ bool FOrderTickTest::RunTest(const FString& Parameters)
 	return bSuccess && !ReportAnyErrors();
 }
 
+UE_DEFINE_GAMEPLAY_TAG_COMMENT(TickTestMessageTag, "AsyncMessages.Internal.test.TickEvent", "Tag for testing async message tick event");
+
+// Simple test message system that is executed manually
+class FTestMessageSystem : public FAsyncMessageSystemBase
+{
+protected:
+	virtual void Startup_Impl() override { }
+	virtual void Shutdown_Impl() override { }
+	virtual void PostQueueMessage(const FAsyncMessageId MessageId, const TArray<FAsyncMessageBindingOptions>& OptionsBoundTo) override { }
+};
+
+// Mass processor for tick testing
+UEngineTickTestProcessor::UEngineTickTestProcessor()
+	: EntityQuery(*this)
+{
+#if WITH_EDITORONLY_DATA
+	bCanShowUpInSettings = false;
+#endif // WITH_EDITORONLY_DATA
+	bAutoRegisterWithProcessingPhases = false;
+	bRequiresGameThreadExecution = false; // Comment out to test async
+	ExecutionFlags = int32(EProcessorExecutionFlags::All);
+
+	ForEachEntityChunkExecutionFunction = [](FMassExecutionContext& Context) {};
+}
+
+void UEngineTickTestProcessor::Execute(FMassEntityManager& EntityManager, FMassExecutionContext& Context)
+{
+	EntityQuery.ForEachEntityChunk(Context, ForEachEntityChunkExecutionFunction);
+}
+
+FGraphEventRef UEngineTickTestProcessor::DispatchProcessorTasks(const TSharedPtr<FMassEntityManager>& EntityManager, FMassExecutionContext& ExecutionContext, const FGraphEventArray& Prerequisites)
+{
+	FGraphEventRef ReturnRef = Super::DispatchProcessorTasks(EntityManager, ExecutionContext, Prerequisites);
+
+	if (!SyncPointName.IsNone())
+	{
+		UE::Tick::FTaskSyncManager* TaskManager = UE::Tick::FTaskSyncManager::Get();
+		if (TaskManager)
+		{
+			UE::Tick::FSyncPointId SyncPoint = TaskManager->FindSyncPoint(EntityManager->GetWorld(), SyncPointName);
+			TaskManager->TriggerSyncPointAfterEvent(SyncPoint, ReturnRef);
+		}
+	}
+
+	return ReturnRef;
+}
+
+static TSharedRef<FMassProcessingPhaseManager> InitializeMassProcessing(FMassEntityManager* EntityManager, UMassProcessor* Processor)
+{
+	TSharedRef<FMassProcessingPhaseManager> PhaseManager = MakeShared<FMassProcessingPhaseManager>();
+
+	FMassProcessingPhaseConfig PhasesConfig[int(EMassProcessingPhase::MAX)];
+	PhaseManager->Initialize(*EntityManager->GetWorld(), PhasesConfig);
+
+	PhaseManager->RegisterDynamicProcessor(*Processor);
+
+	PhaseManager->Start(EntityManager->AsShared());
+
+	return PhaseManager;
+}
+
+// Test using task sync manager to coordinate tick and mass
+IMPLEMENT_CUSTOM_SIMPLE_AUTOMATION_TEST(FTaskSyncTest, FEngineTickTestBase, "System.Engine.Tick.TaskSyncTest", EAutomationTestFlags::EditorContext | EAutomationTestFlags::ClientContext | EAutomationTestFlags::ServerContext | EAutomationTestFlags::EngineFilter)
+bool FTaskSyncTest::RunTest(const FString& Parameters)
+{
+	UE::Tick::FTaskSyncManager* SyncManager = UE::Tick::FTaskSyncManager::Get();
+	if (!SyncManager)
+	{
+		// Skip test for now
+		return true;
+	}
+
+	if (!CreateTestWorld())
+	{
+		return false;
+	}
+
+	float DeltaTime = 0.01f;
+	int32 ActorCount = 1000;
+	bool bSuccess = true;
+
+	bSuccess &= CreateTestActors(ActorCount, AEngineTestTickActor::StaticClass());
+	bSuccess &= BeginPlayInTestWorld();
+
+	if (bSuccess)
+	{
+		check(TestActors.Num() == ActorCount);
+
+		/** Add the test events */
+		FName TestSourceName("EngineTestSource");
+
+		FName TestTaskName("EngineTestTask");
+		FSyncPointDescription TestTaskDescription;
+		TestTaskDescription.RegisteredName = TestTaskName;
+		TestTaskDescription.SourceName = TestSourceName;
+		TestTaskDescription.EventType = ESyncPointEventType::GameThreadTask_HighPriority;
+		TestTaskDescription.ActivationRules = ESyncPointActivationRules::WaitForAllWork;
+		TestTaskDescription.FirstPossibleTickGroup = TG_PrePhysics;
+		TestTaskDescription.LastPossibleTickGroup = TG_PostPhysics;
+		ensure(SyncManager->RegisterNewSyncPoint(TestTaskDescription));
+
+		FName TestEventName("EngineTestEvent");
+		FSyncPointDescription TestEventDescription;
+		TestEventDescription.RegisteredName = TestEventName;
+		TestEventDescription.SourceName = TestSourceName;
+		TestEventDescription.EventType = ESyncPointEventType::SimpleEvent;
+		TestEventDescription.ActivationRules = ESyncPointActivationRules::WaitForTrigger;
+		TestEventDescription.FirstPossibleTickGroup = TG_PrePhysics;
+		TestEventDescription.LastPossibleTickGroup = TG_LastDemotable;
+		ensure(SyncManager->RegisterNewSyncPoint(TestEventDescription));
+		
+		TickTestWorld(); // Tick once to create the sync tick functions
+
+		// Test using requested work
+		{
+			UE::Tick::FSyncPointId TestTaskSyncPoint = SyncManager->FindSyncPoint(GetTestWorld(), TestTaskName);
+			
+			TArray<UE::Tick::FActiveSyncWorkHandle> WorkHandles;
+
+			for (AEngineTestTickActor* TestActor : TestActors)
+			{
+				UE::Tick::FActiveSyncWorkHandle WorkHandle;
+				SyncManager->RegisterWorkHandle(TestTaskSyncPoint, WorkHandle);
+
+				// Disable normal tick
+				TestActor->RegisterAllActorTickFunctions(false, false);
+				
+				ensure(WorkHandle.IsValid());
+
+				WorkHandle.ReserveFutureWork(UE::Tick::ESyncWorkRepetition::Once);
+
+				ensure(WorkHandle.RequestWork(&TestActor->PrimaryActorTick, UE::Tick::ESyncWorkRepetition::EveryFrame));
+				WorkHandles.Add(MoveTemp(WorkHandle));
+			}
+
+			ResetTestActors();
+			TickTestWorld();
+
+			TickTestWorld();
+
+			CheckTickCount(TEXT("RequestedWork"), 2);
+
+			for (AEngineTestTickActor* TestActor : TestActors)
+			{
+				TestActor->RegisterAllActorTickFunctions(true, true);
+			}
+		}
+
+		UE::Tick::FSyncPointId TestEventSyncPoint = SyncManager->FindSyncPoint(GetTestWorld(), TestEventName);
+		FTickFunction* EventTickFunction = SyncManager->GetTickFunctionForSyncPoint(TestEventSyncPoint);
+		check(EventTickFunction);
+
+		for (AEngineTestTickActor* TestActor : TestActors)
+		{
+			// Add a tick dependency
+			TestActor->PrimaryActorTick.AddPrerequisite(TestActor, *EventTickFunction);
+		}
+
+		// Use a message so it happens halfway through tick
+		{
+			auto PayloadTickFunction = [&](const FAsyncMessage& Message)
+			{
+				const FEngineTestTickPayload* Data = Message.GetPayloadData<const FEngineTestTickPayload>();
+				if (Data)
+				{
+				
+					FGraphEventRef FoundEvent;
+
+					FGraphEventRef WaitEvent = FGraphEvent::CreateGraphEvent();
+
+					UE::Tick::FTaskSyncResult Result = SyncManager->GetTaskGraphEvent(TestEventSyncPoint, FoundEvent);
+					ensure(Result);
+					ensure(FoundEvent.IsValid());
+
+					Result = SyncManager->TriggerSyncPointAfterEvent(TestEventSyncPoint, WaitEvent);
+					ensure(Result && Result.WasActivatedForFrame());
+
+					WaitEvent->DispatchSubsequents(ENamedThreads::GameThread);
+
+					Result = SyncManager->TriggerSyncPoint(TestEventSyncPoint);
+					ensure(!Result && Result.WasActivatedForFrame());
+				}
+			};
+
+		
+			const FAsyncMessageId TickTestMessageId = { TickTestMessageTag };
+			TSharedPtr<FAsyncMessageSystemBase> TestSystem = UAsyncMessageWorldSubsystem::GetSharedMessageSystem<FAsyncGameplayMessageSystem>(GetTestWorld());
+			const FAsyncMessageHandle ListenerHandle = TestSystem->BindListener(TickTestMessageId, PayloadTickFunction);
+
+			FEngineTestTickPayload PayloadData;
+			FConstStructView PayloadView = FConstStructView::Make<FEngineTestTickPayload>(PayloadData);
+
+			TestSystem->QueueMessageForBroadcast(TickTestMessageId, PayloadView);
+
+			ResetTestActors();
+			TickTestWorld();
+			CheckTickCount(TEXT("MessageSync"), 1);
+
+			TestSystem->UnbindListener(ListenerHandle);
+		}
+
+		// Test using mass to kickoff event
+		{
+			FMassEntityManager* EntityManager = UE::Mass::Utils::GetEntityManager(GetTestWorld());
+
+			TArray<const UScriptStruct*> Fragments = { FEngineTestTickPayload::StaticStruct() };
+			FMassArchetypeHandle TickArchetype = EntityManager->CreateArchetype(Fragments);
+
+			UEngineTickTestProcessor* Processor = NewObject<UEngineTickTestProcessor>();
+			Processor->CallInitialize(GetTransientPackage(), EntityManager->AsShared());
+			Processor->EntityQuery.AddRequirement<FEngineTestTickPayload>(EMassFragmentAccess::ReadOnly);
+			Processor->SyncPointName = TestEventName;
+
+
+			TArray<FMassEntityHandle> TickEntities;
+			EntityManager->BatchCreateEntities(TickArchetype, TestActors.Num(), TickEntities);
+			TSharedRef<FMassProcessingPhaseManager> PhaseManager = InitializeMassProcessing(EntityManager, Processor);
+
+			ResetTestActors();
+			TickTestWorld();
+			CheckTickCount(TEXT("MassSync"), 1);
+
+			EntityManager->BatchDestroyEntities(TickEntities);
+
+			PhaseManager->Deinitialize();
+		}
+
+		ensure(SyncManager->UnregisterSyncPoint(TestTaskDescription.RegisteredName, TestTaskDescription.SourceName));
+		ensure(SyncManager->UnregisterSyncPoint(TestEventDescription.RegisteredName, TestEventDescription.SourceName));
+		SyncManager->ReloadRegisteredData();
+	}
+
+	bSuccess &= DestroyTestWorld();
+
+	return bSuccess && !ReportAnyErrors();
+}
 
 static TAutoConsoleVariable<int32> CVarEngineTickPerfOptions(
 	TEXT("Automation.Test.EngineTickPerf.Options"),
@@ -528,7 +776,6 @@ bool FPerfTickTest::RunTest(const FString& Parameters)
 		FSimpleMulticastDelegate LambdaDelegate, VirtualLambdaDelegate;
 		FSimpleMulticastDelegate UObjectDelegate, VirtualUObjectDelegate;
 		FSimpleMulticastDelegate WeakLambdaDelegate, VirtualWeakLambdaDelegate;
-		FTSTicker TSTicker;
 
 		for (AEngineTestTickActor* TestActor : TestActors)
 		{
@@ -617,6 +864,7 @@ bool FPerfTickTest::RunTest(const FString& Parameters)
 
 		if (IsTestEnabled(TEXT("WorldTSTicker"), true))
 		{
+			FTSTicker TSTicker;
 			for (AEngineTestTickActor* TestActor : TestActors)
 			{
 				TSTicker.AddTicker(FTickerDelegate::CreateWeakLambda(TestActor, [TestActor](float) {TestActor->VirtualTick(); return true; }), 0.0f);
@@ -632,8 +880,152 @@ bool FPerfTickTest::RunTest(const FString& Parameters)
 				}
 			}
 			CheckTickCount(TEXT("WorldTSTicker"), TickCount);
+			TSTicker.Reset();
+		}
+		
+		if (IsTestEnabled(TEXT("GameplayMessageSystem"), false))
+		{
+			const FAsyncMessageId TickTestMessageId = { TickTestMessageTag };
+			auto PayloadTickFunction = [](const FAsyncMessage& Message)
+			{
+				const FEngineTestTickPayload* Data = Message.GetPayloadData<const FEngineTestTickPayload>();
+				if (Data)
+				{
+					AEngineTestTickActor* TargetActor = Data->TargetActor.Get();
+					if (ensure(TargetActor))
+					{
+						TargetActor->VirtualTick();
+					}
+				}
+			};
+
+			TSharedPtr<FAsyncGameplayMessageSystem> GameplaySystem = UAsyncMessageWorldSubsystem::GetSharedMessageSystem<FAsyncGameplayMessageSystem>(GetTestWorld());
+			if (ensure(GameplaySystem))
+			{
+				const FAsyncMessageHandle ListenerHandle = GameplaySystem->BindListener(TickTestMessageId, PayloadTickFunction);
+
+				FEngineTestTickPayload PayloadData;
+				FConstStructView PayloadView = FConstStructView::Make<FEngineTestTickPayload>(PayloadData);
+
+				ResetTestActors();
+				{
+					LOG_SCOPE_TIME(GameplayMessageSystem);
+					for (int32 i = 0; i < TickCount; i++)
+					{
+						for (AEngineTestTickActor* TestActor : TestActors)
+						{
+							PayloadData.TargetActor = TestActor;
+							GameplaySystem->QueueMessageForBroadcast(TickTestMessageId, PayloadView);
+						}
+
+						// This will run the tick group message tick handlers
+						TickTestWorld();
+					}
+				}
+				CheckTickCount(TEXT("GameplayMessageSystem"), TickCount);
+
+				GameplaySystem->UnbindListener(ListenerHandle);
+			}
+		}
+
+		if (IsTestEnabled(TEXT("MassProcessor"), false))
+		{
+			FMassEntityManager* EntityManager = UE::Mass::Utils::GetEntityManager(GetTestWorld());
+
+			TArray<const UScriptStruct*> Fragments = { FEngineTestTickPayload::StaticStruct() };
+			FMassArchetypeHandle TickArchetype = EntityManager->CreateArchetype(Fragments);
+
+			UEngineTickTestProcessor* Processor = NewObject<UEngineTickTestProcessor>();
+			Processor->CallInitialize(GetTransientPackage(), EntityManager->AsShared());
+			Processor->EntityQuery.AddRequirement<FEngineTestTickPayload>(EMassFragmentAccess::ReadOnly);
+			Processor->ForEachEntityChunkExecutionFunction = [](FMassExecutionContext& Context)
+			{
+				const TConstArrayView<FEngineTestTickPayload> TickPayloads = Context.GetFragmentView<FEngineTestTickPayload>();
+				for (const FEngineTestTickPayload& Payload : TickPayloads)
+				{
+					AEngineTestTickActor* TargetActor = Payload.TargetActor.Get();
+					if (ensure(TargetActor))
+					{
+						TargetActor->VirtualTick();
+					}
+				}
+			};
+
+			TArray<FMassEntityHandle> TickEntities;
+			EntityManager->BatchCreateEntities(TickArchetype, TestActors.Num(), TickEntities);
+
+			for (int32 i = 0; i < TestActors.Num(); i++)
+			{
+				check(TickEntities.IsValidIndex(i));
+
+				// TODO more efficient way of initializing?
+				FEngineTestTickPayload& Payload = EntityManager->GetFragmentDataChecked<FEngineTestTickPayload>(TickEntities[i]);
+				Payload.TargetActor = TestActors[i];
+			}
+
+			TSharedRef<FMassProcessingPhaseManager> PhaseManager = InitializeMassProcessing(EntityManager, Processor);
+
+			ResetTestActors();
+			{
+				LOG_SCOPE_TIME(MassProcessor);
+				for (int32 i = 0; i < TickCount; i++)
+				{
+					// This will run the tick group message tick handlers
+					TickTestWorld();
+				}
+			}
+			CheckTickCount(TEXT("MassProcessor"), TickCount);
+
+			EntityManager->BatchDestroyEntities(TickEntities);
+			
+			PhaseManager->Deinitialize();
+		}
+
+		UE::Tick::FTaskSyncManager* SyncManager = UE::Tick::FTaskSyncManager::Get();
+		if (IsTestEnabled(TEXT("TaskSyncManager"), true) && SyncManager)
+		{
+			FName TestSourceName("EngineTestSource");
+
+			FName TestTaskName("EngineTestTask");
+			FSyncPointDescription TestTaskDescription;
+			TestTaskDescription.RegisteredName = TestTaskName;
+			TestTaskDescription.SourceName = TestSourceName;
+			TestTaskDescription.EventType = ESyncPointEventType::GameThreadTask;
+			TestTaskDescription.ActivationRules = ESyncPointActivationRules::AlwaysActivate;
+			TestTaskDescription.FirstPossibleTickGroup = TG_PrePhysics;
+			TestTaskDescription.LastPossibleTickGroup = TG_PrePhysics;
+			ensure(SyncManager->RegisterNewSyncPoint(TestTaskDescription));
+
+			// Tick once to create events
+			TickTestWorld();
+
+			UE::Tick::FSyncPointId TestTaskSyncPoint = SyncManager->FindSyncPoint(GetTestWorld(), TestTaskName);
+
+			TArray<UE::Tick::FActiveSyncWorkHandle> WorkHandles;
+
+			for (AEngineTestTickActor* TestActor : TestActors)
+			{
+				UE::Tick::FActiveSyncWorkHandle WorkHandle;
+				SyncManager->RegisterWorkHandle(TestTaskSyncPoint, WorkHandle);
+
+				ensure(WorkHandle.RequestWork(&TestActor->PrimaryActorTick, UE::Tick::ESyncWorkRepetition::EveryFrame));
+				WorkHandles.Add(MoveTemp(WorkHandle));
+			}
+
+			ResetTestActors();
+			{
+				LOG_SCOPE_TIME(TaskSyncManager);
+				for (int32 i = 0; i < TickCount; i++)
+				{
+					TickTestWorld();
+				}
+			}
+			CheckTickCount(TEXT("TaskSyncManager"), TickCount);
+
+			WorkHandles.Reset();
+			ensure(SyncManager->UnregisterSyncPoint(TestTaskDescription.RegisteredName, TestTaskDescription.SourceName));
+			SyncManager->ReloadRegisteredData();
 		}
-		TSTicker.Reset();
 
 		if (IsTestEnabled(TEXT("WorldTimerManager"), true))
 		{
diff --git a/Engine/Plugins/Tests/RuntimeTests/Source/RuntimeTests/Public/EngineRuntimeTests.h b/Engine/Plugins/Tests/RuntimeTests/Source/RuntimeTests/Public/EngineRuntimeTests.h
--- a/Engine/Plugins/Tests/RuntimeTests/Source/RuntimeTests/Public/EngineRuntimeTests.h
+++ b/Engine/Plugins/Tests/RuntimeTests/Source/RuntimeTests/Public/EngineRuntimeTests.h
@@ -4,6 +4,8 @@
 #include "GameFramework/Actor.h"
 #include "Misc/AutomationTest.h"
 #include "Tests/AutomationCommon.h"
+#include "MassProcessor.h"
+#include "MassEntityTypes.h"
 
 #include "EngineRuntimeTests.generated.h"
 
@@ -66,6 +68,45 @@ private:
 	TObjectPtr<UBillboardComponent> SpriteComponent;
 };
 
+USTRUCT(NotBlueprintType, meta=(Hidden))
+struct RUNTIMETESTS_API FEngineTestTickPayload : public FMassFragment
+{
+	GENERATED_BODY()
+	
+	TWeakObjectPtr<AEngineTestTickActor> TargetActor = nullptr;
+};
+
+// Copied from MassEntityTestSuite 
+UCLASS()
+class RUNTIMETESTS_API UEngineTickTestProcessor : public UMassProcessor
+{
+	GENERATED_BODY()
+public:
+	UEngineTickTestProcessor();
+	FMassProcessorExecutionOrder& GetMutableExecutionOrder() { return ExecutionOrder; }
+	virtual void Execute(FMassEntityManager& EntityManager, FMassExecutionContext& Context) override;
+	virtual FGraphEventRef DispatchProcessorTasks(const TSharedPtr<FMassEntityManager>& EntityManager, FMassExecutionContext& ExecutionContext, const FGraphEventArray& Prerequisites = FGraphEventArray()) override;
+
+	virtual bool ShouldAllowQueryBasedPruning(const bool bRuntimeMode = true) const { return false; }
+	virtual void ConfigureQueries(const TSharedRef<FMassEntityManager>& EntityManager) override {}
+
+	using FExecutionFunction = TFunction<void(FMassEntityManager& EntityManager, FMassExecutionContext& Context)>;
+	FExecutionFunction ExecutionFunction;
+
+	/** 
+	 * By default ExecutionFunction is configured to pass this function over to EntityQuery.ForEachEntityChunk call. 
+	 * Note that this function won't be used if you override ExecutionFunction's default value.
+	 */
+	FMassExecuteFunction ForEachEntityChunkExecutionFunction;
+
+	/** public on purpose, this is a test processor, no worries about access*/
+	FMassEntityQuery EntityQuery;
+
+	/** SyncPoint event to emit after processor task has finished */
+	FName SyncPointName;
+};
+
+
 #if WITH_AUTOMATION_WORKER
 
 /** Automation test base class that wraps a test world and handles checking tick counts */
diff --git a/Engine/Plugins/Tests/RuntimeTests/Source/RuntimeTests/RuntimeTests.Build.cs b/Engine/Plugins/Tests/RuntimeTests/Source/RuntimeTests/RuntimeTests.Build.cs
--- a/Engine/Plugins/Tests/RuntimeTests/Source/RuntimeTests/RuntimeTests.Build.cs
+++ b/Engine/Plugins/Tests/RuntimeTests/Source/RuntimeTests/RuntimeTests.Build.cs
@@ -10,6 +10,7 @@ public class RuntimeTests : ModuleRules
 			new string[]
 			{
 				"Core",
+				"MassEntity",
 			}
 		);
 			
@@ -21,6 +22,8 @@ public class RuntimeTests : ModuleRules
 				"Engine",
 				"Slate",
 				"SlateCore",
+				"AsyncMessageSystem",
+				"GameplayTags",
 			}
 		);
 
