Enable UHT namespace support for Engine, testing this before enabling for all modules.
Move the TaskSyncManager structs and enums into UE::Tick namespace for consistency with native types.

diff --git a/Engine/Plugins/Tests/RuntimeTests/Source/RuntimeTests/Private/EngineRuntimeTests.cpp b/Engine/Plugins/Tests/RuntimeTests/Source/RuntimeTests/Private/EngineRuntimeTests.cpp
--- a/Engine/Plugins/Tests/RuntimeTests/Source/RuntimeTests/Private/EngineRuntimeTests.cpp
+++ b/Engine/Plugins/Tests/RuntimeTests/Source/RuntimeTests/Private/EngineRuntimeTests.cpp
@@ -490,21 +490,21 @@ bool FTaskSyncTest::RunTest(const FString& Parameters)
 		FName TestSourceName("EngineTestSource");
 
 		FName TestTaskName("EngineTestTask");
-		FSyncPointDescription TestTaskDescription;
+		UE::Tick::FSyncPointDescription TestTaskDescription;
 		TestTaskDescription.RegisteredName = TestTaskName;
 		TestTaskDescription.SourceName = TestSourceName;
-		TestTaskDescription.EventType = ESyncPointEventType::GameThreadTask_HighPriority;
-		TestTaskDescription.ActivationRules = ESyncPointActivationRules::WaitForAllWork;
+		TestTaskDescription.EventType = UE::Tick::ESyncPointEventType::GameThreadTask_HighPriority;
+		TestTaskDescription.ActivationRules = UE::Tick::ESyncPointActivationRules::WaitForAllWork;
 		TestTaskDescription.FirstPossibleTickGroup = TG_PrePhysics;
 		TestTaskDescription.LastPossibleTickGroup = TG_PostPhysics;
 		ensure(SyncManager->RegisterNewSyncPoint(TestTaskDescription));
 
 		FName TestEventName("EngineTestEvent");
-		FSyncPointDescription TestEventDescription;
+		UE::Tick::FSyncPointDescription TestEventDescription;
 		TestEventDescription.RegisteredName = TestEventName;
 		TestEventDescription.SourceName = TestSourceName;
-		TestEventDescription.EventType = ESyncPointEventType::SimpleEvent;
-		TestEventDescription.ActivationRules = ESyncPointActivationRules::WaitForTrigger;
+		TestEventDescription.EventType = UE::Tick::ESyncPointEventType::SimpleEvent;
+		TestEventDescription.ActivationRules = UE::Tick::ESyncPointActivationRules::WaitForTrigger;
 		TestEventDescription.FirstPossibleTickGroup = TG_PrePhysics;
 		TestEventDescription.LastPossibleTickGroup = TG_LastDemotable;
 		ensure(SyncManager->RegisterNewSyncPoint(TestEventDescription));
@@ -987,11 +987,11 @@ bool FPerfTickTest::RunTest(const FString& Parameters)
 			FName TestSourceName("EngineTestSource");
 
 			FName TestTaskName("EngineTestTask");
-			FSyncPointDescription TestTaskDescription;
+			UE::Tick::FSyncPointDescription TestTaskDescription;
 			TestTaskDescription.RegisteredName = TestTaskName;
 			TestTaskDescription.SourceName = TestSourceName;
-			TestTaskDescription.EventType = ESyncPointEventType::GameThreadTask;
-			TestTaskDescription.ActivationRules = ESyncPointActivationRules::AlwaysActivate;
+			TestTaskDescription.EventType = UE::Tick::ESyncPointEventType::GameThreadTask;
+			TestTaskDescription.ActivationRules = UE::Tick::ESyncPointActivationRules::AlwaysActivate;
 			TestTaskDescription.FirstPossibleTickGroup = TG_PrePhysics;
 			TestTaskDescription.LastPossibleTickGroup = TG_PrePhysics;
 			ensure(SyncManager->RegisterNewSyncPoint(TestTaskDescription));
diff --git a/Engine/Source/Runtime/Engine/Engine.Build.cs b/Engine/Source/Runtime/Engine/Engine.Build.cs
--- a/Engine/Source/Runtime/Engine/Engine.Build.cs
+++ b/Engine/Source/Runtime/Engine/Engine.Build.cs
@@ -509,5 +509,7 @@ public class Engine : ModuleRules
 		SetupIrisSupport(Target, bIrisAddAsPublicDependency);
 
 		PrivateDefinitions.Add("UE_DEPRECATE_LEGACY_MATH_CONSTANT_MACRO_NAMES=1");
+
+		bAllowUETypesInNamespaces = true;
 	}
 }
